Mover 2.0 fixing issue with control rotation pitch and walking/falling mode causing actor to stop moving

diff --git a/Engine/Plugins/Experimental/MoverExamples/Source/MoverExamples/Private/MoverExamplesCharacter.cpp b/Engine/Plugins/Experimental/MoverExamples/Source/MoverExamples/Private/MoverExamplesCharacter.cpp
--- a/Engine/Plugins/Experimental/MoverExamples/Source/MoverExamples/Private/MoverExamplesCharacter.cpp
+++ b/Engine/Plugins/Experimental/MoverExamples/Source/MoverExamples/Private/MoverExamplesCharacter.cpp
@@ -54,6 +54,17 @@ void AMoverExamplesCharacter::Tick(float DeltaTime)
 	CachedLookInput = FRotator::ZeroRotator;
 }
 
+void AMoverExamplesCharacter::BeginPlay()
+{
+	Super::BeginPlay();
+
+	if (APlayerController* PC = Cast<APlayerController>(Controller))
+	{
+		PC->PlayerCameraManager->ViewPitchMax = 89.0f;
+		PC->PlayerCameraManager->ViewPitchMin = -89.0f;
+	}
+}
+
 // Called to bind functionality to input
 void AMoverExamplesCharacter::SetupPlayerInputComponent(UInputComponent* PlayerInputComponent)
 {
@@ -133,7 +144,20 @@ void AMoverExamplesCharacter::OnProduceInput(float DeltaMs, FMoverInputCmdContex
 
 	if (bUsingInputIntentForMove)
 	{
-		DefaultKinematicInputs.SetMoveInput(EMoveInputType::DirectionalIntent, DefaultKinematicInputs.ControlRotation.RotateVector(CachedMoveInputIntent));
+		FRotator Rotator = DefaultKinematicInputs.ControlRotation;
+		FVector FinalDirectionalIntent;
+		if (const UMoverComponent* MoverComponent = GetMoverComponent())
+		{
+			if (MoverComponent->IsOnGround() || MoverComponent->IsFalling())
+			{
+				const FVector RotationProjectedOntoUpDirection = FVector::VectorPlaneProject(Rotator.Vector(), MoverComponent->GetUpDirection()).GetSafeNormal();
+				Rotator = RotationProjectedOntoUpDirection.Rotation();
+			}
+
+			FinalDirectionalIntent = Rotator.RotateVector(CachedMoveInputIntent);
+		}
+		
+		DefaultKinematicInputs.SetMoveInput(EMoveInputType::DirectionalIntent, FinalDirectionalIntent);
 	}
 	else
 	{
diff --git a/Engine/Plugins/Experimental/MoverExamples/Source/MoverExamples/Public/MoverExamplesCharacter.h b/Engine/Plugins/Experimental/MoverExamples/Source/MoverExamples/Public/MoverExamplesCharacter.h
--- a/Engine/Plugins/Experimental/MoverExamples/Source/MoverExamples/Public/MoverExamplesCharacter.h
+++ b/Engine/Plugins/Experimental/MoverExamples/Source/MoverExamples/Public/MoverExamplesCharacter.h
@@ -24,6 +24,8 @@ public:
 	// Called every frame
 	virtual void Tick(float DeltaTime) override;
 
+	virtual void BeginPlay() override;
+
 	// Called to bind functionality to input
 	virtual void SetupPlayerInputComponent(class UInputComponent* PlayerInputComponent) override;
 
