MoverExamples: first pass ziplining mechanic example, added to Dashing pawn with 1 zipline in the main examples map. Still needs animation support.

#virtualized

diff --git a/Engine/Plugins/Experimental/MoverExamples/Source/MoverExamples/Private/CharacterVariants/Ziplining/ZipliningMode.cpp b/Engine/Plugins/Experimental/MoverExamples/Source/MoverExamples/Private/CharacterVariants/Ziplining/ZipliningMode.cpp
new file mode 100644
--- /dev/null
+++ b/Engine/Plugins/Experimental/MoverExamples/Source/MoverExamples/Private/CharacterVariants/Ziplining/ZipliningMode.cpp
@@ -0,0 +1,247 @@
+// Copyright Epic Games, Inc. All Rights Reserved.
+
+#include "CharacterVariants/Ziplining/ZipliningMode.h"
+#include "MoveLibrary/MovementUtils.h"
+#include "MoverComponent.h"
+#include "CharacterVariants/Ziplining/ZiplineInterface.h"
+#include "CharacterVariants/Ziplining/ZipliningTransitions.h"
+#include "Kismet/KismetSystemLibrary.h"
+#include "Kinematic/Settings/CommonLegacyMovementSettings.h"
+#include "MoverLog.h"
+
+
+#include UE_INLINE_GENERATED_CPP_BY_NAME(ZipliningMode)
+
+
+
+
+
+// FZipliningState //////////////////////////////
+
+FMoverDataStructBase* FZipliningState::Clone() const
+{
+	FZipliningState* CopyPtr = new FZipliningState(*this);
+	return CopyPtr;
+}
+
+bool FZipliningState::NetSerialize(FArchive& Ar, UPackageMap* Map, bool& bOutSuccess)
+{
+	Super::NetSerialize(Ar, Map, bOutSuccess);
+
+	Ar << ZiplineActor;
+	Ar.SerializeBits(&bIsMovingAtoB,1);
+
+	bOutSuccess = true;
+	return true;
+}
+
+void FZipliningState::ToString(FAnsiStringBuilderBase& Out) const
+{
+	Super::ToString(Out);
+
+	Out.Appendf("ZiplineActor: %s\n", *GetNameSafe(ZiplineActor));
+	Out.Appendf("IsMovingAtoB: %d\n", bIsMovingAtoB);
+}
+
+bool FZipliningState::ShouldReconcile(const FMoverDataStructBase& AuthorityState) const
+{
+	const FZipliningState* AuthorityZiplineState = static_cast<const FZipliningState*>(&AuthorityState);
+
+	return (ZiplineActor != AuthorityZiplineState->ZiplineActor) ||
+		   (bIsMovingAtoB != AuthorityZiplineState->bIsMovingAtoB);
+}
+
+void FZipliningState::Interpolate(const FMoverDataStructBase& From, const FMoverDataStructBase& To, float Pct)
+{
+	const FZipliningState* FromState = static_cast<const FZipliningState*>(&From);
+	const FZipliningState* ToState = static_cast<const FZipliningState*>(&To);
+
+	ZiplineActor = ToState->ZiplineActor;
+	bIsMovingAtoB = ToState->bIsMovingAtoB;
+}
+
+
+
+// UZipliningMode //////////////////////////////
+
+UZipliningMode::UZipliningMode(const FObjectInitializer& ObjectInitializer)
+	: Super(ObjectInitializer)
+{
+	Transitions.Add(CreateDefaultSubobject<UZiplineEndTransition>(TEXT("ZiplineEndTransition")));
+}
+
+
+void UZipliningMode::OnGenerateMove(const FMoverTickStartData& StartState, const FMoverTimeStep& TimeStep, FProposedMove& OutProposedMove) const
+{
+	UMoverComponent* MoverComp = GetMoverComponent();
+
+	// Ziplining is just following a path from A to B, so all movement is handled in OnSimulationTick
+	OutProposedMove = FProposedMove();
+
+}
+
+void UZipliningMode::OnSimulationTick(const FSimulationTickParams& Params, FMoverTickEndData& OutputState)
+{
+	// Are we continuing a move or starting fresh?
+	const FZipliningState* StartingZipState = Params.StartState.SyncState.SyncStateCollection.FindDataByType<FZipliningState>();
+
+	FMoverDefaultSyncState& OutputSyncState = OutputState.SyncState.SyncStateCollection.FindOrAddMutableDataByType<FMoverDefaultSyncState>();
+	FZipliningState& OutZipState            = OutputState.SyncState.SyncStateCollection.FindOrAddMutableDataByType<FZipliningState>();
+
+
+	UMoverComponent* MoverComp = Params.MoverComponent;
+	USceneComponent* UpdatedComponent = Params.UpdatedComponent;
+	UPrimitiveComponent* UpdatedPrimitive = Params.UpdatedPrimitive;
+	AActor* MoverActor = MoverComp->GetOwner();
+
+	USceneComponent* StartPoint = nullptr;
+	USceneComponent* EndPoint = nullptr;
+	FVector ZipDirection;
+	FVector FlatFacingDir;
+
+	const float DeltaSeconds = Params.TimeStep.StepMs * 0.001f;
+
+	FVector ActorOrigin;
+	FVector BoxExtent;
+	MoverActor->GetActorBounds(true, OUT ActorOrigin, OUT BoxExtent);
+	const FVector ActorToZiplineOffset = MoverComp->GetUpDirection() * BoxExtent.Z;
+
+	if (!StartingZipState)
+	{
+		// There is no existing zipline state... so let's find the target
+		//    A) teleport to the closest starting point, set the zip direction
+		//    B) choose the appropriate facing direction
+		//    C) choose the appropriate initial velocity
+		TArray<AActor*> OverlappingActors;
+		MoverComp->GetOwner()->GetOverlappingActors(OUT OverlappingActors);
+
+		for (AActor* CandidateActor : OverlappingActors)
+		{
+			bool bIsZipline = UKismetSystemLibrary::DoesImplementInterface(CandidateActor, UZipline::StaticClass());
+
+			if (bIsZipline)
+			{
+				const FVector MoverLoc = UpdatedComponent->GetComponentLocation();
+				USceneComponent* ZipPointA = IZipline::Execute_GetStartComponent(CandidateActor);
+				USceneComponent* ZipPointB = IZipline::Execute_GetEndComponent(CandidateActor);
+
+				if (FVector::DistSquared(ZipPointA->GetComponentLocation(), MoverLoc) < FVector::DistSquared(ZipPointB->GetComponentLocation(), MoverLoc))
+				{
+					OutZipState.bIsMovingAtoB = true;
+					StartPoint = ZipPointA;
+					EndPoint = ZipPointB;
+				}
+				else
+				{
+					OutZipState.bIsMovingAtoB = false;
+					StartPoint = ZipPointB;
+					EndPoint = ZipPointA;
+				}
+
+				ZipDirection = (EndPoint->GetComponentLocation() - StartPoint->GetComponentLocation()).GetSafeNormal();
+
+				const FVector WarpLocation = StartPoint->GetComponentLocation() - ActorToZiplineOffset;
+
+				FlatFacingDir = FVector::VectorPlaneProject(ZipDirection, MoverComp->GetUpDirection()).GetSafeNormal();
+
+				OutZipState.ZiplineActor = CandidateActor;
+
+				UpdatedComponent->GetOwner()->TeleportTo(WarpLocation, FlatFacingDir.ToOrientationRotator());
+
+				break;
+			}
+		}
+
+		// If we were unable to find a valid target zipline, refund all the time and let the actor fall
+		if (!StartPoint || !EndPoint)
+		{
+			FName DefaultAirMode = KinematicModeNames::Falling;
+			if (UCommonLegacyMovementSettings* LegacySettings = MoverComp->FindSharedSettings_Mutable<UCommonLegacyMovementSettings>())
+			{
+				DefaultAirMode = LegacySettings->AirMovementModeName;
+			}
+
+			OutputState.MovementEndState.NextModeName = KinematicModeNames::Falling;
+			OutputState.MovementEndState.RemainingMs = Params.TimeStep.StepMs;
+			return;
+		}
+
+	}
+	else
+	{
+		check(StartingZipState->ZiplineActor);
+		OutZipState = *StartingZipState;
+
+		USceneComponent* ZipPointA = IZipline::Execute_GetStartComponent(StartingZipState->ZiplineActor);
+		USceneComponent* ZipPointB = IZipline::Execute_GetEndComponent(StartingZipState->ZiplineActor);
+
+		if (StartingZipState->bIsMovingAtoB)
+		{
+			StartPoint = ZipPointA;
+			EndPoint = ZipPointB;
+		}
+		else
+		{
+			StartPoint = ZipPointB;
+			EndPoint = ZipPointA;
+		}
+
+		ZipDirection = (EndPoint->GetComponentLocation() - StartPoint->GetComponentLocation()).GetSafeNormal();
+		FlatFacingDir = FVector::VectorPlaneProject(ZipDirection, MoverComp->GetUpDirection()).GetSafeNormal();
+	}
+
+
+	// Now let's slide along the zipline
+	const FVector StepStartPos = UpdatedComponent->GetComponentLocation() + ActorToZiplineOffset;
+	const FVector DesiredEndPos = StepStartPos + (ZipDirection * MaxSpeed * DeltaSeconds);	// TODO: Make speed more dynamic
+
+	FVector ActualEndPos = FMath::ClosestPointOnSegment(DesiredEndPos,
+		StartPoint->GetComponentLocation(),
+		EndPoint->GetComponentLocation());
+
+	bool bWillReachEndPosition = (ActualEndPos - EndPoint->GetComponentLocation()).IsNearlyZero();
+
+	FVector MoveDelta = ActualEndPos - StepStartPos;
+
+
+
+	FMovementRecord MoveRecord;
+	MoveRecord.SetDeltaSeconds(DeltaSeconds);
+
+
+	if (!MoveDelta.IsNearlyZero())
+	{
+		
+		FHitResult Hit(1.f);
+
+		UMovementUtils::TrySafeMoveUpdatedComponent(UpdatedComponent, UpdatedPrimitive, MoveDelta, FlatFacingDir.ToOrientationQuat(), true, Hit, ETeleportType::None, MoveRecord);
+	}
+
+
+	const FVector FinalLocation = UpdatedComponent->GetComponentLocation();
+	const FVector FinalVelocity = MoveRecord.GetRelevantVelocity();
+
+	OutputSyncState.SetTransforms_WorldSpace(FinalLocation,
+		UpdatedComponent->GetComponentRotation(),
+		FinalVelocity,
+		nullptr); // no movement base
+
+	UpdatedComponent->ComponentVelocity = FinalVelocity;
+
+
+	if (bWillReachEndPosition)
+	{
+		FName DefaultAirMode = KinematicModeNames::Falling;
+		if (UCommonLegacyMovementSettings* LegacySettings = MoverComp->FindSharedSettings_Mutable<UCommonLegacyMovementSettings>())
+		{
+			DefaultAirMode = LegacySettings->AirMovementModeName;
+		}
+
+		OutputState.MovementEndState.NextModeName = DefaultAirMode;
+		// TODO: If we reach the end position early, we should refund the remaining time
+
+	}
+}
+
+
+
diff --git a/Engine/Plugins/Experimental/MoverExamples/Source/MoverExamples/Private/CharacterVariants/Ziplining/ZipliningTransitions.cpp b/Engine/Plugins/Experimental/MoverExamples/Source/MoverExamples/Private/CharacterVariants/Ziplining/ZipliningTransitions.cpp
new file mode 100644
--- /dev/null
+++ b/Engine/Plugins/Experimental/MoverExamples/Source/MoverExamples/Private/CharacterVariants/Ziplining/ZipliningTransitions.cpp
@@ -0,0 +1,80 @@
+// Copyright Epic Games, Inc. All Rights Reserved.
+
+#include "CharacterVariants/Ziplining/ZipliningTransitions.h"
+#include "CharacterVariants/AbilityInputs.h"
+#include "CharacterVariants/Ziplining/ZiplineInterface.h"
+#include "MoverComponent.h"
+#include "GameFramework/Actor.h"
+#include "Kismet/KismetSystemLibrary.h"
+
+
+// UZiplineStartTransition //////////////////////////////
+
+UZiplineStartTransition::UZiplineStartTransition(const FObjectInitializer& ObjectInitializer)
+	: Super(ObjectInitializer)
+{
+}
+
+FTransitionEvalResult UZiplineStartTransition::OnEvaluate(const FSimulationTickParams& Params) const
+{
+	FTransitionEvalResult EvalResult = FTransitionEvalResult::NoTransition;
+
+	UMoverComponent* MoverComp = Params.MoverComponent;
+
+	const FMoverDefaultSyncState* SyncState = Params.StartState.SyncState.SyncStateCollection.FindDataByType<FMoverDefaultSyncState>();
+
+	if (MoverComp && MoverComp->IsAirborne() && SyncState && SyncState->MovementMode != ZipliningModeName)
+	{
+		if (const FMoverExampleAbilityInputs* AbilityInputs = Params.StartState.InputCmd.InputCollection.FindDataByType<FMoverExampleAbilityInputs>())
+		{
+			if (AbilityInputs->bWantsToStartZiplining)
+			{
+				TArray<AActor*> OverlappingActors;
+				MoverComp->GetOwner()->GetOverlappingActors(OUT OverlappingActors);
+
+				for (AActor* CandidateActor : OverlappingActors)
+				{
+					bool bIsZipline = UKismetSystemLibrary::DoesImplementInterface(CandidateActor, UZipline::StaticClass());
+
+					if (bIsZipline)
+					{
+						EvalResult.NextMode = ZipliningModeName;
+						break;
+					}
+				}
+			}
+		}
+	}
+
+	return EvalResult;
+}
+
+
+
+// UZiplineEndTransition //////////////////////////////
+
+UZiplineEndTransition::UZiplineEndTransition(const FObjectInitializer& ObjectInitializer)
+	: Super(ObjectInitializer)
+{
+}
+
+
+FTransitionEvalResult UZiplineEndTransition::OnEvaluate(const FSimulationTickParams& Params) const
+{
+	FTransitionEvalResult EvalResult = FTransitionEvalResult::NoTransition;
+
+	if (const FKinematicDefaultInputs* DefaultInputs = Params.StartState.InputCmd.InputCollection.FindDataByType<FKinematicDefaultInputs>())
+	{
+		if (DefaultInputs->bIsJumpJustPressed)
+		{
+			EvalResult.NextMode = AutoExitToMode;
+		}
+	}
+
+	return EvalResult;
+}
+
+void UZiplineEndTransition::OnTrigger(const FSimulationTickParams& Params)
+{
+	//TODO: create a small jump, using current directionality
+}
diff --git a/Engine/Plugins/Experimental/MoverExamples/Source/MoverExamples/Public/CharacterVariants/AbilityInputs.h b/Engine/Plugins/Experimental/MoverExamples/Source/MoverExamples/Public/CharacterVariants/AbilityInputs.h
--- a/Engine/Plugins/Experimental/MoverExamples/Source/MoverExamples/Public/CharacterVariants/AbilityInputs.h
+++ b/Engine/Plugins/Experimental/MoverExamples/Source/MoverExamples/Public/CharacterVariants/AbilityInputs.h
@@ -21,6 +21,10 @@ struct MOVEREXAMPLES_API FMoverExampleAbilityInputs : public FMoverDataStructBas
 	UPROPERTY(BlueprintReadWrite, Category = Mover)
 	bool bIsVaultJustPressed = false;
 
+	UPROPERTY(BlueprintReadWrite, Category = Mover)
+	bool bWantsToStartZiplining = false;
+
+
 	// @return newly allocated copy of this FKinematicDefaultInputs. Must be overridden by child classes
 	virtual FMoverDataStructBase* Clone() const override
 	{
@@ -36,6 +40,7 @@ struct MOVEREXAMPLES_API FMoverExampleAbilityInputs : public FMoverDataStructBas
 		Ar.SerializeBits(&bIsDashJustPressed, 1);
 		Ar.SerializeBits(&bIsAimPressed, 1);
 		Ar.SerializeBits(&bIsVaultJustPressed, 1);
+		Ar.SerializeBits(&bWantsToStartZiplining, 1);
 
 		bOutSuccess = true;
 		return true;
@@ -49,6 +54,7 @@ struct MOVEREXAMPLES_API FMoverExampleAbilityInputs : public FMoverDataStructBas
 		Out.Appendf("bIsDashJustPressed: %i\n", bIsDashJustPressed);
 		Out.Appendf("bIsAimPressed: %i\n", bIsAimPressed);
 		Out.Appendf("bIsVaultJustPressed: %i\n", bIsVaultJustPressed);
+		Out.Appendf("bWantsToStartZiplining: %i\n", bWantsToStartZiplining);
 	}
 
 	virtual void AddReferencedObjects(FReferenceCollector& Collector) override { Super::AddReferencedObjects(Collector); }
diff --git a/Engine/Plugins/Experimental/MoverExamples/Source/MoverExamples/Public/CharacterVariants/Ziplining/ZiplineInterface.h b/Engine/Plugins/Experimental/MoverExamples/Source/MoverExamples/Public/CharacterVariants/Ziplining/ZiplineInterface.h
new file mode 100644
--- /dev/null
+++ b/Engine/Plugins/Experimental/MoverExamples/Source/MoverExamples/Public/CharacterVariants/Ziplining/ZiplineInterface.h
@@ -0,0 +1,30 @@
+// Copyright Epic Games, Inc. All Rights Reserved.
+
+#pragma once
+
+#include "CoreMinimal.h"
+#include "UObject/Interface.h"
+#include "ZiplineInterface.generated.h"
+
+
+
+UINTERFACE(MinimalAPI)
+class UZipline : public UInterface
+{
+	GENERATED_BODY()
+};
+
+
+/** Interface that can be implemented on an object supporting simple zipline traversal between 2 points, 
+    represented as Scene Components. */
+class IZipline
+{
+	GENERATED_BODY()
+public:
+
+	UFUNCTION(BlueprintCallable, BlueprintNativeEvent, Category = "Zipline")
+	USceneComponent* GetStartComponent();
+
+	UFUNCTION(BlueprintCallable, BlueprintNativeEvent, Category = "Zipline")
+	USceneComponent* GetEndComponent();
+};
diff --git a/Engine/Plugins/Experimental/MoverExamples/Source/MoverExamples/Public/CharacterVariants/Ziplining/ZipliningMode.h b/Engine/Plugins/Experimental/MoverExamples/Source/MoverExamples/Public/CharacterVariants/Ziplining/ZipliningMode.h
new file mode 100644
--- /dev/null
+++ b/Engine/Plugins/Experimental/MoverExamples/Source/MoverExamples/Public/CharacterVariants/Ziplining/ZipliningMode.h
@@ -0,0 +1,52 @@
+// Copyright Epic Games, Inc. All Rights Reserved.
+
+#pragma once
+
+#include "CoreMinimal.h"
+#include "MovementMode.h"
+#include "ZipliningMode.generated.h"
+
+
+namespace ExtendedModeNames
+{
+	const FName Ziplining = TEXT("Ziplining");
+}
+
+
+// ZipliningMode: movement mode that traverses an actor implementing the IZipline interface
+UCLASS(Blueprintable, BlueprintType)
+class UZipliningMode : public UBaseMovementMode
+{
+	GENERATED_UCLASS_BODY()
+
+public:
+	virtual void OnGenerateMove(const FMoverTickStartData& StartState, const FMoverTimeStep& TimeStep, FProposedMove& OutProposedMove) const override;
+	virtual void OnSimulationTick(const FSimulationTickParams& Params, FMoverTickEndData& OutputState) override;
+
+	// Maximum speed 
+	UPROPERTY(EditAnywhere, Category = "Ziplining", meta = (ClampMin = "1", UIMin = "1", ForceUnits = "cm/s"))
+	float MaxSpeed = 1000.0f;
+};
+
+
+// Data block containing ziplining state info, used while ZipliningMode is active
+USTRUCT()
+struct FZipliningState : public FMoverDataStructBase
+{
+	GENERATED_USTRUCT_BODY()
+
+	TObjectPtr<AActor> ZiplineActor;
+	bool bIsMovingAtoB;
+
+	FZipliningState()
+		: bIsMovingAtoB(true)
+	{
+	}
+
+	virtual FMoverDataStructBase* Clone() const override;
+	virtual bool NetSerialize(FArchive& Ar, UPackageMap* Map, bool& bOutSuccess) override;
+	virtual UScriptStruct* GetScriptStruct() const override { return StaticStruct(); }
+	virtual void ToString(FAnsiStringBuilderBase& Out) const override;
+	virtual bool ShouldReconcile(const FMoverDataStructBase& AuthorityState) const override;
+	virtual void Interpolate(const FMoverDataStructBase& From, const FMoverDataStructBase& To, float Pct) override;
+};
diff --git a/Engine/Plugins/Experimental/MoverExamples/Source/MoverExamples/Public/CharacterVariants/Ziplining/ZipliningTransitions.h b/Engine/Plugins/Experimental/MoverExamples/Source/MoverExamples/Public/CharacterVariants/Ziplining/ZipliningTransitions.h
new file mode 100644
--- /dev/null
+++ b/Engine/Plugins/Experimental/MoverExamples/Source/MoverExamples/Public/CharacterVariants/Ziplining/ZipliningTransitions.h
@@ -0,0 +1,40 @@
+// Copyright Epic Games, Inc. All Rights Reserved.
+
+#pragma once
+
+#include "MovementModeTransition.h"
+#include "CharacterVariants/Ziplining/ZipliningMode.h"
+#include "ZipliningTransitions.generated.h"
+
+/**
+ * Transition that handles starting ziplining based on input. Character must be airborne to catch the
+ * zipline, regardless of input.
+ */
+UCLASS(Blueprintable, BlueprintType)
+class UZiplineStartTransition : public UBaseMovementModeTransition
+{
+	GENERATED_UCLASS_BODY()
+
+public:
+	virtual FTransitionEvalResult OnEvaluate(const FSimulationTickParams& Params) const override;
+
+	UPROPERTY(EditAnywhere, Category = "Ziplining")
+	FName ZipliningModeName = ExtendedModeNames::Ziplining;
+};
+
+/**
+ * Transition that handles exiting ziplining based on input
+ */
+UCLASS(Blueprintable, BlueprintType)
+class UZiplineEndTransition : public UBaseMovementModeTransition
+{
+	GENERATED_UCLASS_BODY()
+
+public:
+	virtual FTransitionEvalResult OnEvaluate(const FSimulationTickParams& Params) const override;
+	virtual void OnTrigger(const FSimulationTickParams& Params) override;
+
+	// Mode to enter when exiting the zipline
+	UPROPERTY(EditAnywhere, Category = "Ziplining")
+	FName AutoExitToMode = KinematicModeNames::Falling;
+};
