Mover: removing unnecessary and misleading MoverBasePawn class, since there's no need for MoverComponent to attached to any particular Actor type. What little functionality it provided has been transferred to MoverExamplesCharacter.

diff --git a/Engine/Plugins/Experimental/Mover/Source/Mover/Private/MoverBasePawn.cpp b/Engine/Plugins/Experimental/Mover/Source/Mover/Private/MoverBasePawn.cpp
deleted file mode 100644
diff --git a/Engine/Plugins/Experimental/Mover/Source/Mover/Public/MoverBasePawn.h b/Engine/Plugins/Experimental/Mover/Source/Mover/Public/MoverBasePawn.h
deleted file mode 100644
diff --git a/Engine/Plugins/Experimental/MoverExamples/Source/MoverExamples/Private/MoverExamplesCharacter.cpp b/Engine/Plugins/Experimental/MoverExamples/Source/MoverExamples/Private/MoverExamplesCharacter.cpp
--- a/Engine/Plugins/Experimental/MoverExamples/Source/MoverExamples/Private/MoverExamplesCharacter.cpp
+++ b/Engine/Plugins/Experimental/MoverExamples/Source/MoverExamples/Private/MoverExamplesCharacter.cpp
@@ -16,9 +16,18 @@
 #include "EnhancedInputComponent.h"
 #include "InputAction.h"
 
+static const FName Name_CharacterMotionComponent(TEXT("MoverComponent"));
+
 AMoverExamplesCharacter::AMoverExamplesCharacter(const FObjectInitializer& ObjectInitializer)
 	: Super(ObjectInitializer)
 {
+	CharacterMotionComponent = CreateDefaultSubobject<UMoverComponent>(Name_CharacterMotionComponent);
+	ensure(CharacterMotionComponent);
+
+	PrimaryActorTick.bCanEverTick = true;
+
+	SetReplicatingMovement(false);	// disable Actor-level movement replication, since our Mover component will handle it
+
 	auto IsImplementedInBlueprint = [](const UFunction* Func) -> bool
 	{
 		return Func && ensure(Func->GetOuter())
@@ -146,11 +155,11 @@ void AMoverExamplesCharacter::OnProduceInput(float DeltaMs, FMoverInputCmdContex
 	{
 		FRotator Rotator = CharacterInputs.ControlRotation;
 		FVector FinalDirectionalIntent;
-		if (const UMoverComponent* MoverComponent = GetMoverComponent())
+		if (const UMoverComponent* MoverComp = GetMoverComponent())
 		{
-			if (MoverComponent->IsOnGround() || MoverComponent->IsFalling())
+			if (MoverComp->IsOnGround() || MoverComp->IsFalling())
 			{
-				const FVector RotationProjectedOntoUpDirection = FVector::VectorPlaneProject(Rotator.Vector(), MoverComponent->GetUpDirection()).GetSafeNormal();
+				const FVector RotationProjectedOntoUpDirection = FVector::VectorPlaneProject(Rotator.Vector(), MoverComp->GetUpDirection()).GetSafeNormal();
 				Rotator = RotationProjectedOntoUpDirection.Rotation();
 			}
 
diff --git a/Engine/Plugins/Experimental/MoverExamples/Source/MoverExamples/Public/MoverExamplesCharacter.h b/Engine/Plugins/Experimental/MoverExamples/Source/MoverExamples/Public/MoverExamplesCharacter.h
--- a/Engine/Plugins/Experimental/MoverExamples/Source/MoverExamples/Public/MoverExamplesCharacter.h
+++ b/Engine/Plugins/Experimental/MoverExamples/Source/MoverExamples/Public/MoverExamplesCharacter.h
@@ -3,15 +3,15 @@
 #pragma once
 
 #include "CoreMinimal.h"
-#include "MoverBasePawn.h"
 #include "MoverSimulationTypes.h"
+#include "GameFramework/Pawn.h"
 #include "MoverExamplesCharacter.generated.h"
 
 class UInputAction;
 struct FInputActionValue;
 
 UCLASS()
-class MOVEREXAMPLES_API AMoverExamplesCharacter : public AMoverBasePawn, public IMoverInputProducerInterface
+class MOVEREXAMPLES_API AMoverExamplesCharacter : public APawn, public IMoverInputProducerInterface
 {
 	GENERATED_BODY()
 
@@ -29,6 +29,10 @@ public:
 	// Called to bind functionality to input
 	virtual void SetupPlayerInputComponent(class UInputComponent* PlayerInputComponent) override;
 
+	// Accessor for the actor's movement component
+	UFUNCTION(BlueprintPure, Category = Mover)
+	UMoverComponent* GetMoverComponent() const { return CharacterMotionComponent; }
+
 	// Request the character starts moving with an intended directional magnitude. A length of 1 indicates maximum acceleration.
 	UFUNCTION(BlueprintCallable, Category=MoverExamples)
 	virtual void RequestMoveByIntent(const FVector& DesiredIntent) { CachedMoveInputIntent = DesiredIntent; }
@@ -88,6 +92,10 @@ public:
 	UPROPERTY(EditAnywhere, BlueprintReadWrite, Category = MoverExamples)
 	bool bMaintainLastInputOrientation = false;
 
+protected:
+	UPROPERTY(Category = Movement, VisibleAnywhere, BlueprintReadOnly, meta = (AllowPrivateAccess = "true"))
+	TObjectPtr<UMoverComponent> CharacterMotionComponent;
+
 private:
 	FVector LastAffirmativeMoveInput = FVector::ZeroVector;	// Movement input (intent or velocity) the last time we had one that wasn't zero
 
