MoverExamples: removing erroneous error message that was being logged during one of the map startups

#rnx

diff --git a/Engine/Plugins/Experimental/MoverExamples/Source/MoverExamples/Private/MovementBases/FollowSplineMode.cpp b/Engine/Plugins/Experimental/MoverExamples/Source/MoverExamples/Private/MovementBases/FollowSplineMode.cpp
--- a/Engine/Plugins/Experimental/MoverExamples/Source/MoverExamples/Private/MovementBases/FollowSplineMode.cpp
+++ b/Engine/Plugins/Experimental/MoverExamples/Source/MoverExamples/Private/MovementBases/FollowSplineMode.cpp
@@ -157,10 +157,6 @@ void UFollowSplineMode::SetControlSpline(const AActor* SplineProviderActor, FSpl
 
 		CurrentSplineTime = FMath::Clamp(InitialOffsetTime, StartOffsetTime, EndOffsetTime) - StartOffsetTime;
 	}
-	else
-	{
-		UE_LOG(LogMover, Error, TEXT("No Valid Spline component was provided. Make sure that the spline provider has a valid spline."));
-	}
 }
 
 void UFollowSplineMode::OnRegistered(const FName ModeName)
@@ -172,13 +168,13 @@ void UFollowSplineMode::OnRegistered(const FName ModeName)
 
 void UFollowSplineMode::ConfigureSplineData()
 {
-	// Control Spline is set externally. 
+	// Control Spline is already set
 	if (ControlSpline)
 	{
 		return;
 	}
 
-	USplineComponent* SplineComponent = nullptr;
+	// Attempt to find a spline component on the actor who owns this mode
 	const UObject* OwnerObject = GetOuter();
 	while (OwnerObject != nullptr)
 	{
