Prevent MoverExamplesCharacter from being instantiated as an actor type. It is our native base class for MoverExamples plugin's pawns and doesn't have the required scene components.

diff --git a/Engine/Plugins/Experimental/MoverExamples/Source/MoverExamples/Public/MoverExamplesCharacter.h b/Engine/Plugins/Experimental/MoverExamples/Source/MoverExamples/Public/MoverExamplesCharacter.h
--- a/Engine/Plugins/Experimental/MoverExamples/Source/MoverExamples/Public/MoverExamplesCharacter.h
+++ b/Engine/Plugins/Experimental/MoverExamples/Source/MoverExamples/Public/MoverExamplesCharacter.h
@@ -11,7 +11,12 @@ class UInputAction;
 class UCharacterMoverComponent;
 struct FInputActionValue;
 
-UCLASS()
+/** 
+ * MoverExamplesCharacter: the base pawn class used by the MoverExamples plugin. Handles coalescing of input events.
+ * Cannot be instantiated on its own.
+ */ 
+
+UCLASS(Abstract)
 class MOVEREXAMPLES_API AMoverExamplesCharacter : public APawn, public IMoverInputProducerInterface
 {
 	GENERATED_BODY()
