Mover examples removing some arbitrary gravity code in the MoverExamplesCharacter that tried to account for gravity while rotated. New rotation will always assum +Z is up and we'll need something custom to handle camera movement if we want it relative to gravity direction.

#rnx

diff --git a/Engine/Plugins/Experimental/MoverExamples/Source/MoverExamples/Private/MoverExamplesCharacter.cpp b/Engine/Plugins/Experimental/MoverExamples/Source/MoverExamples/Private/MoverExamplesCharacter.cpp
--- a/Engine/Plugins/Experimental/MoverExamples/Source/MoverExamples/Private/MoverExamplesCharacter.cpp
+++ b/Engine/Plugins/Experimental/MoverExamples/Source/MoverExamples/Private/MoverExamplesCharacter.cpp
@@ -61,15 +61,9 @@ void AMoverExamplesCharacter::Tick(float DeltaTime)
 		// Simple input scaling. A real game will probably map this to an acceleration curve
 		static float LookRateYaw = 100.f;	// degs/sec
 		static float LookRatePitch = 100.f;	// degs/sec
-		static float LookRateRoll = 100.f;	// degs/sec
 
-		const FRotator GravityRotation(CharacterMotionComponent->GetWorldToGravityTransform() * CachedLookInput.Quaternion());
-		const FRotator DefaultRotation(CharacterMotionComponent->GetWorldToGravityTransform() * FRotator::ZeroRotator.Quaternion());
-		const FRotator Difference = GravityRotation - DefaultRotation;
-			
-		PC->AddYawInput(Difference.Yaw * LookRateYaw * DeltaTime);
-		PC->AddPitchInput(-Difference.Pitch * LookRatePitch * DeltaTime);
-		PC->AddRollInput(Difference.Roll * LookRateRoll * DeltaTime);
+		PC->AddYawInput(CachedLookInput.Yaw * LookRateYaw * DeltaTime);
+		PC->AddPitchInput(-CachedLookInput.Pitch * LookRatePitch * DeltaTime);
 	}
 
 	// Clear all camera-related cached input
