Code and content fixes for mover examples to work as expected.
Implementing FMoverDataStructBase functions (ShouldReconcile, Merge and Interpolate) for FMoverExampleAbilityInputs to prevent base class implementation from asserting. Also, setting Physics Manny Pawn Extended replication mode to resim and adding a network physics settings component so those settings can be modified for it.

diff --git a/Engine/Plugins/Experimental/MoverExamples/Source/MoverExamples/Public/CharacterVariants/AbilityInputs.h b/Engine/Plugins/Experimental/MoverExamples/Source/MoverExamples/Public/CharacterVariants/AbilityInputs.h
--- a/Engine/Plugins/Experimental/MoverExamples/Source/MoverExamples/Public/CharacterVariants/AbilityInputs.h
+++ b/Engine/Plugins/Experimental/MoverExamples/Source/MoverExamples/Public/CharacterVariants/AbilityInputs.h
@@ -27,6 +27,40 @@ struct MOVEREXAMPLES_API FMoverExampleAbilityInputs : public FMoverDataStructBas
 	UPROPERTY(BlueprintReadWrite, Category = Mover)
 	bool bWantsToBeCrouched = false;
 
+	// Implementation of FMoverDataStructBase
+	virtual bool ShouldReconcile(const FMoverDataStructBase& AuthorityState) const override
+	{
+		const FMoverExampleAbilityInputs& TypedAuthority = static_cast<const FMoverExampleAbilityInputs&>(AuthorityState);
+		return (TypedAuthority.bIsDashJustPressed != bIsDashJustPressed)
+			|| (TypedAuthority.bIsAimPressed != bIsAimPressed)
+			|| (TypedAuthority.bIsVaultJustPressed != bIsVaultJustPressed)
+			|| (TypedAuthority.bWantsToStartZiplining != bWantsToStartZiplining)
+			|| (TypedAuthority.bWantsToBeCrouched != bWantsToBeCrouched);
+	}
+
+	virtual void Interpolate(const FMoverDataStructBase& From, const FMoverDataStructBase& To, float LerpFactor) override
+	{
+		// Since we're just copying bool properties, we simply copy them from From if LerpFactor is less than 0.5, otherwise from To
+		const FMoverExampleAbilityInputs& SourceAbilityInputs = static_cast<const FMoverExampleAbilityInputs&>((LerpFactor < 0.5f) ? From : To);
+
+		bIsDashJustPressed = SourceAbilityInputs.bIsDashJustPressed;
+		bIsAimPressed = SourceAbilityInputs.bIsAimPressed;
+		bIsVaultJustPressed = SourceAbilityInputs.bIsVaultJustPressed;
+		bWantsToStartZiplining = SourceAbilityInputs.bWantsToStartZiplining;
+		bWantsToBeCrouched = SourceAbilityInputs.bWantsToBeCrouched;
+	}
+
+	virtual void Merge(const FMoverDataStructBase& From) override
+	{
+		const FMoverExampleAbilityInputs& TypedFrom = static_cast<const FMoverExampleAbilityInputs&>(From);
+
+		bIsDashJustPressed |= TypedFrom.bIsDashJustPressed;
+		bIsAimPressed |= TypedFrom.bIsAimPressed;
+		bIsVaultJustPressed |= TypedFrom.bIsVaultJustPressed;
+		bWantsToStartZiplining |= TypedFrom.bWantsToStartZiplining;
+		bWantsToBeCrouched |= TypedFrom.bWantsToBeCrouched;
+	}
+
 	// @return newly allocated copy of this FMoverExampleAbilityInputs. Must be overridden by child classes
 	virtual FMoverDataStructBase* Clone() const override
 	{
