MoverExamples:
- Removing the hard coded CharacterMoverComponent, as it is not the correct mover component for physics based character movement
- Introducing BaseAnimatedMannyPawn that has all the input setup but does not impose a concrete mover component class
- Fixing physics based examples to work, using PhysicsCharacterMoverComponent
- Crouch is implemented in a new MoverExamplesPhysicsCharacterMoverComponent
- Adding a ChaosMannyPawnExtended, which is currently using PhysicsCharacterMoverComponent
- Eventually we would like to remove PhysicsCharacterMoverComponent and MoverExamplesPhysicsCharacterMoverComponent completely and maybe even CharacterMoverComponent and completely specialize behavior based on the MoverComponent setup (modes, transitions, settings, etc.)

#virtualized

diff --git a/Engine/Plugins/Experimental/Mover/Source/Mover/Public/DefaultMovementSet/CharacterMoverComponent.h b/Engine/Plugins/Experimental/Mover/Source/Mover/Public/DefaultMovementSet/CharacterMoverComponent.h
--- a/Engine/Plugins/Experimental/Mover/Source/Mover/Public/DefaultMovementSet/CharacterMoverComponent.h
+++ b/Engine/Plugins/Experimental/Mover/Source/Mover/Public/DefaultMovementSet/CharacterMoverComponent.h
@@ -14,7 +14,7 @@
 DECLARE_DYNAMIC_MULTICAST_DELEGATE_TwoParams(FMover_OnStanceChanged, EStanceMode, OldStance, EStanceMode, NewStance);
 
 
-UCLASS(BlueprintType, meta = (BlueprintSpawnableComponent))
+UCLASS(BlueprintType, Blueprintable, meta = (BlueprintSpawnableComponent))
 class MOVER_API UCharacterMoverComponent : public UMoverComponent
 {
 	GENERATED_BODY()
diff --git a/Engine/Plugins/Experimental/Mover/Source/Mover/Public/PhysicsMover/PhysicsCharacterMoverComponent.h b/Engine/Plugins/Experimental/Mover/Source/Mover/Public/PhysicsMover/PhysicsCharacterMoverComponent.h
--- a/Engine/Plugins/Experimental/Mover/Source/Mover/Public/PhysicsMover/PhysicsCharacterMoverComponent.h
+++ b/Engine/Plugins/Experimental/Mover/Source/Mover/Public/PhysicsMover/PhysicsCharacterMoverComponent.h
@@ -6,7 +6,7 @@
 #include "PhysicsCharacterMoverComponent.generated.h"
 
 
-UCLASS(BlueprintType, meta = (BlueprintSpawnableComponent))
+UCLASS(BlueprintType, Blueprintable, meta = (BlueprintSpawnableComponent))
 class MOVER_API UPhysicsCharacterMoverComponent : public UCharacterMoverComponent
 {
 	GENERATED_BODY()
diff --git a/Engine/Plugins/Experimental/MoverExamples/MoverExamples.uplugin b/Engine/Plugins/Experimental/MoverExamples/MoverExamples.uplugin
--- a/Engine/Plugins/Experimental/MoverExamples/MoverExamples.uplugin
+++ b/Engine/Plugins/Experimental/MoverExamples/MoverExamples.uplugin
@@ -27,6 +27,10 @@
 			"Name": "Mover",
 			"Enabled": true
 		},
+		{
+			"Name": "ChaosMover",
+			"Enabled": true
+		},
 		{
 			"Name": "EnhancedInput",
 			"Enabled": true
diff --git a/Engine/Plugins/Experimental/MoverExamples/Source/MoverExamples/Private/CharacterVariants/MoverExamplesPhysicsCharacterMoverComponent.cpp b/Engine/Plugins/Experimental/MoverExamples/Source/MoverExamples/Private/CharacterVariants/MoverExamplesPhysicsCharacterMoverComponent.cpp
new file mode 100644
--- /dev/null
+++ b/Engine/Plugins/Experimental/MoverExamples/Source/MoverExamples/Private/CharacterVariants/MoverExamplesPhysicsCharacterMoverComponent.cpp
@@ -0,0 +1,22 @@
+// Copyright Epic Games, Inc. All Rights Reserved.
+
+#include "CharacterVariants/MoverExamplesPhysicsCharacterMoverComponent.h"
+
+#include "CharacterVariants/AbilityInputs.h"
+
+void UMoverExamplesPhysicsCharacterMoverComponent::OnMoverPreMovement(const FMoverTimeStep& TimeStep, const FMoverInputCmdContext& InputCmd, const FMoverSyncState& SyncState, const FMoverAuxStateContext& AuxState)
+{
+	if (FMoverExampleAbilityInputs* AbilityInputs = InputCmd.InputCollection.FindMutableDataByType<FMoverExampleAbilityInputs>())
+	{
+		if (AbilityInputs->bWantsToBeCrouched)
+		{
+			Crouch_Internal(SyncState);
+		}
+		else
+		{
+			UnCrouch_Internal(SyncState);
+		}
+	}
+
+	Super::OnMoverPreMovement(TimeStep, InputCmd, SyncState, AuxState);
+}
\ No newline at end of file
diff --git a/Engine/Plugins/Experimental/MoverExamples/Source/MoverExamples/Private/MoverExamplesCharacter.cpp b/Engine/Plugins/Experimental/MoverExamples/Source/MoverExamples/Private/MoverExamplesCharacter.cpp
--- a/Engine/Plugins/Experimental/MoverExamples/Source/MoverExamples/Private/MoverExamplesCharacter.cpp
+++ b/Engine/Plugins/Experimental/MoverExamples/Source/MoverExamples/Private/MoverExamplesCharacter.cpp
@@ -17,15 +17,10 @@
 #include "InputAction.h"
 #include "DefaultMovementSet/NavMoverComponent.h"
 
-static const FName Name_CharacterMotionComponent(TEXT("MoverComponent"));
-
 AMoverExamplesCharacter::AMoverExamplesCharacter(const FObjectInitializer& ObjectInitializer)
 	: Super(ObjectInitializer)
 	, NavMoverComponent(nullptr)
 {
-	CharacterMotionComponent = CreateDefaultSubobject<UCharacterMoverComponent>(Name_CharacterMotionComponent);
-	ensure(CharacterMotionComponent);
-
 	PrimaryActorTick.bCanEverTick = true;
 
 	SetReplicatingMovement(false);	// disable Actor-level movement replication, since our Mover component will handle it
@@ -39,14 +34,23 @@ AMoverExamplesCharacter::AMoverExamplesCharacter(const FObjectInitializer& Objec
 	static FName ProduceInputBPFuncName = FName(TEXT("OnProduceInputInBlueprint"));
 	UFunction* ProduceInputFunction = GetClass()->FindFunctionByName(ProduceInputBPFuncName);
 	bHasProduceInputinBpFunc = IsImplementedInBlueprint(ProduceInputFunction);
+}
+
+void AMoverExamplesCharacter::PostInitializeComponents()
+{
+	Super::PostInitializeComponents();
+
+	CharacterMotionComponent = FindComponentByClass<UCharacterMoverComponent>();
 
-	if (USceneComponent* UpdatedComponent = CharacterMotionComponent->GetUpdatedComponent())
+	if (CharacterMotionComponent)
 	{
-		UpdatedComponent->SetCanEverAffectNavigation(bCanAffectNavigationGeneration);
+		if (USceneComponent* UpdatedComponent = CharacterMotionComponent->GetUpdatedComponent())
+		{
+			UpdatedComponent->SetCanEverAffectNavigation(bCanAffectNavigationGeneration);
+		}
 	}
 }
 
-
 // Called every frame
 void AMoverExamplesCharacter::Tick(float DeltaTime)
 {
diff --git a/Engine/Plugins/Experimental/MoverExamples/Source/MoverExamples/Public/CharacterVariants/MoverExamplesPhysicsCharacterMoverComponent.h b/Engine/Plugins/Experimental/MoverExamples/Source/MoverExamples/Public/CharacterVariants/MoverExamplesPhysicsCharacterMoverComponent.h
new file mode 100644
--- /dev/null
+++ b/Engine/Plugins/Experimental/MoverExamples/Source/MoverExamples/Public/CharacterVariants/MoverExamplesPhysicsCharacterMoverComponent.h
@@ -0,0 +1,18 @@
+// Copyright Epic Games, Inc. All Rights Reserved.
+
+#pragma once
+
+#include "PhysicsMover/PhysicsCharacterMoverComponent.h"
+
+#include "MoverExamplesPhysicsCharacterMoverComponent.generated.h"
+
+
+UCLASS(BlueprintType, Blueprintable, meta = (BlueprintSpawnableComponent))
+class MOVEREXAMPLES_API UMoverExamplesPhysicsCharacterMoverComponent : public UPhysicsCharacterMoverComponent
+{
+	GENERATED_BODY()
+
+protected:
+
+	virtual void OnMoverPreMovement(const FMoverTimeStep& TimeStep, const FMoverInputCmdContext& InputCmd, const FMoverSyncState& SyncState, const FMoverAuxStateContext& AuxState) override;
+};
diff --git a/Engine/Plugins/Experimental/MoverExamples/Source/MoverExamples/Public/MoverExamplesCharacter.h b/Engine/Plugins/Experimental/MoverExamples/Source/MoverExamples/Public/MoverExamplesCharacter.h
--- a/Engine/Plugins/Experimental/MoverExamples/Source/MoverExamples/Public/MoverExamplesCharacter.h
+++ b/Engine/Plugins/Experimental/MoverExamples/Source/MoverExamples/Public/MoverExamplesCharacter.h
@@ -33,6 +33,8 @@ public:
 
 	virtual void BeginPlay() override;
 
+	virtual void PostInitializeComponents() override;
+
 	// Called to bind functionality to input
 	virtual void SetupPlayerInputComponent(class UInputComponent* PlayerInputComponent) override;
 
