Making CharacterMotionComponent Transient in MoverExamplesCharacter.h

diff --git a/Engine/Plugins/Experimental/MoverExamples/Source/MoverExamples/Public/MoverExamplesCharacter.h b/Engine/Plugins/Experimental/MoverExamples/Source/MoverExamples/Public/MoverExamplesCharacter.h
--- a/Engine/Plugins/Experimental/MoverExamples/Source/MoverExamples/Public/MoverExamplesCharacter.h
+++ b/Engine/Plugins/Experimental/MoverExamples/Source/MoverExamples/Public/MoverExamplesCharacter.h
@@ -102,7 +102,7 @@ public:
 	bool bMaintainLastInputOrientation = false;
 
 protected:
-	UPROPERTY(Category = Movement, VisibleAnywhere, BlueprintReadOnly, meta = (AllowPrivateAccess = "true"))
+	UPROPERTY(Category = Movement, VisibleAnywhere, BlueprintReadOnly, Transient, meta = (AllowPrivateAccess = "true"))
 	TObjectPtr<UCharacterMoverComponent> CharacterMotionComponent;
 
 	/** Holds functionality for nav movement data and functions */
