MoverExamples: improving handling of orientation intent when not orienting to the movement direction

#Synced-CL 50587857

diff --git a/Engine/Plugins/Experimental/MoverExamples/Source/MoverExamples/Private/MoverExamplesCharacter.cpp b/Engine/Plugins/Experimental/MoverExamples/Source/MoverExamples/Private/MoverExamplesCharacter.cpp
--- a/Engine/Plugins/Experimental/MoverExamples/Source/MoverExamples/Private/MoverExamplesCharacter.cpp
+++ b/Engine/Plugins/Experimental/MoverExamples/Source/MoverExamples/Private/MoverExamplesCharacter.cpp
@@ -237,12 +237,12 @@ void AMoverExamplesCharacter::OnProduceInput(float DeltaMs, FMoverInputCmdContex
 			CharacterInputs.OrientationIntent = CharacterInputs.ControlRotation.Vector().GetSafeNormal();
 		}
 
-		LastAffirmativeMoveInput = CharacterInputs.GetMoveInput();
+		LastAffirmativeOrientationIntent = CharacterInputs.OrientationIntent;
 	}
 	else if (bMaintainLastInputOrientation)
 	{
 		// There is no movement intent, so use the last-known affirmative move input
-		CharacterInputs.OrientationIntent = LastAffirmativeMoveInput;
+		CharacterInputs.OrientationIntent = LastAffirmativeOrientationIntent;
 	}
 	
 	CharacterInputs.bIsJumpPressed = bIsJumpPressed;
diff --git a/Engine/Plugins/Experimental/MoverExamples/Source/MoverExamples/Public/MoverExamplesCharacter.h b/Engine/Plugins/Experimental/MoverExamples/Source/MoverExamples/Public/MoverExamplesCharacter.h
--- a/Engine/Plugins/Experimental/MoverExamples/Source/MoverExamples/Public/MoverExamplesCharacter.h
+++ b/Engine/Plugins/Experimental/MoverExamples/Source/MoverExamples/Public/MoverExamplesCharacter.h
@@ -110,7 +110,7 @@ protected:
 	TObjectPtr<UNavMoverComponent> NavMoverComponent;
 	
 private:
-	FVector LastAffirmativeMoveInput = FVector::ZeroVector;	// Movement input (intent or velocity) the last time we had one that wasn't zero
+	FVector LastAffirmativeOrientationIntent = FVector::ZeroVector;	// Orientation direction at the last time we had one that wasn't zero
 
 	FVector CachedMoveInputIntent = FVector::ZeroVector;
 	FVector CachedMoveInputVelocity = FVector::ZeroVector;
