Mover: Initializing some FollowPathMode data in mode Activate so we guaranteed run it on servers and clients. This fixes a bug where control points sometimes never got initialized in forward predicted movement.

#Synced-CL 50999505

diff --git a/Engine/Plugins/Experimental/MoverExamples/Source/MoverExamples/Private/MovementBases/FollowPathMode.cpp b/Engine/Plugins/Experimental/MoverExamples/Source/MoverExamples/Private/MovementBases/FollowPathMode.cpp
--- a/Engine/Plugins/Experimental/MoverExamples/Source/MoverExamples/Private/MovementBases/FollowPathMode.cpp
+++ b/Engine/Plugins/Experimental/MoverExamples/Source/MoverExamples/Private/MovementBases/FollowPathMode.cpp
@@ -37,8 +37,6 @@ void UFollowPathMode::SimulationTick_Implementation(const FSimulationTickParams&
 	// If we don't already have a valid pathing state, we need to initialize 
 	if (!StartingPathState || !StartingPathState->HasValidPathState())
 	{
-		UpdateControlPoints(true);
-
 		// Indicates we haven't started pathing yet. Finalize control points and capture origins
 		OutputPathState.BaseLocation = UpdatedComponent->GetComponentLocation();
 		OutputPathState.CurrentPathPos = 0.f;
@@ -57,8 +55,6 @@ void UFollowPathMode::SimulationTick_Implementation(const FSimulationTickParams&
 		// Move the component to the first path location
 		FHitResult IgnoredHit(1.f);
 		UpdatedComponent->MoveComponent(StartingLocation - OutputPathState.BaseLocation, UpdatedComponent->GetComponentRotation(), false, &IgnoredHit);
-
-		TimeMultiplier = 1.f / Duration;
 	}
 
 	float RemainingSecs = DeltaSeconds;
@@ -109,6 +105,14 @@ void UFollowPathMode::SimulationTick_Implementation(const FSimulationTickParams&
 	UpdatedComponent->ComponentVelocity = Velocity;
 }
 
+void UFollowPathMode::Activate(const FMoverEventContext& Context)
+{
+	Super::Activate(Context);
+	
+	UpdateControlPoints(true);
+	TimeMultiplier = 1.f / (Duration > UE_KINDA_SMALL_NUMBER ? Duration : UE_KINDA_SMALL_NUMBER);
+}
+
 void UFollowPathMode::UpdateControlPoints(bool InForceUpdate)
 {
 	const USceneComponent* UpdatedComponent = nullptr;
diff --git a/Engine/Plugins/Experimental/MoverExamples/Source/MoverExamples/Public/MovementBases/FollowPathMode.h b/Engine/Plugins/Experimental/MoverExamples/Source/MoverExamples/Public/MovementBases/FollowPathMode.h
--- a/Engine/Plugins/Experimental/MoverExamples/Source/MoverExamples/Public/MovementBases/FollowPathMode.h
+++ b/Engine/Plugins/Experimental/MoverExamples/Source/MoverExamples/Public/MovementBases/FollowPathMode.h
@@ -39,6 +39,8 @@ public:
 	virtual void GenerateMove_Implementation(const FMoverTickStartData& StartState, const FMoverTimeStep& TimeStep, FProposedMove& OutProposedMove) const override;
 	virtual void SimulationTick_Implementation(const FSimulationTickParams& Params, FMoverTickEndData& OutputState) override;
 
+	virtual void Activate(const FMoverEventContext& Context) override;
+	
 	// List of ordered path locations to visit
 	UPROPERTY(EditAnywhere, BlueprintReadWrite, Category = Pathing)
 	TArray<FInterpControlPoint> ControlPoints;
