MoverExamples: fixed crash that could occur if giving FollowSplineMode a bad/missing spline actor during setup time

#pf todo
#Synced-CL 51816243

diff --git a/Engine/Plugins/Experimental/MoverExamples/Source/MoverExamples/Private/MovementBases/FollowSplineMode.cpp b/Engine/Plugins/Experimental/MoverExamples/Source/MoverExamples/Private/MovementBases/FollowSplineMode.cpp
--- a/Engine/Plugins/Experimental/MoverExamples/Source/MoverExamples/Private/MovementBases/FollowSplineMode.cpp
+++ b/Engine/Plugins/Experimental/MoverExamples/Source/MoverExamples/Private/MovementBases/FollowSplineMode.cpp
@@ -150,16 +150,19 @@ void UFollowSplineMode::SimulationTick_Implementation(const FSimulationTickParam
 
 void UFollowSplineMode::SetControlSpline(const AActor* SplineProviderActor, FSplineOffsetRangeInput Offset)
 {
-	if (USplineComponent* SplineComponent = SplineProviderActor->GetComponentByClass<USplineComponent>())
+	if (SplineProviderActor)
 	{
-		using namespace FollowSplineMode::Utils::Private;
-		ControlSpline = SplineComponent;
+		if (USplineComponent* SplineComponent = SplineProviderActor->GetComponentByClass<USplineComponent>())
+		{
+			using namespace FollowSplineMode::Utils::Private;
+			ControlSpline = SplineComponent;
 		
-		const float StartOffsetTime = ComputeRangeInputValue(ControlSpline, StartOffset);
-		const float EndOffsetTime = ComputeRangeInputValue(ControlSpline, EndOffset);
-		const float InitialOffsetTime = ComputeRangeInputValue(ControlSpline, Offset);
+			const float StartOffsetTime = ComputeRangeInputValue(ControlSpline, StartOffset);
+			const float EndOffsetTime = ComputeRangeInputValue(ControlSpline, EndOffset);
+			const float InitialOffsetTime = ComputeRangeInputValue(ControlSpline, Offset);
 
-		CurrentSplineTime = FMath::Clamp(InitialOffsetTime, StartOffsetTime, EndOffsetTime) - StartOffsetTime;
+			CurrentSplineTime = FMath::Clamp(InitialOffsetTime, StartOffsetTime, EndOffsetTime) - StartOffsetTime;
+		}
 	}
 }
 
